#fix 	Fixed buffer size being incorrectly specifed in the XML parser.
#detail 	Fix for crash reported in AnswerHub thread: https://answers.unrealengine.com/questions/16540/bug-conecting-to-source-control-causes-project-to.html

diff --git a/Engine/Source/Developer/XmlParser/Private/XmlFile.cpp b/Engine/Source/Developer/XmlParser/Private/XmlFile.cpp
--- a/Engine/Source/Developer/XmlParser/Private/XmlFile.cpp
+++ b/Engine/Source/Developer/XmlParser/Private/XmlFile.cpp
@@ -71,7 +71,7 @@ bool FXmlFile::LoadFile(const FString& InFile, EConstructMethod::Type ConstructM
 	else
 	{
 		// Parse input buffer into an array of lines
-		SplitLines(Input, *InFile, sizeof(TCHAR) * InFile.Len());
+		SplitLines(Input, *InFile, InFile.Len());
 	}
 
 	// Pre-process the input
