XML parser: fixed out-of-bounds memory access and made saving cross-platform by using UTF-8

Local function SplitLines was using BufferSize as if it were the character count (the bug had already been worked around for in-memory XML, but file reading was broken). This happened to work correctly unless it crashed due to the out-of-bounds memory access causing a page fault.

The switch to UTF-8 addresses TTP 336433 (CrashReportClient was silently failing to write the user's id to the crash report XML file).

diff --git a/Engine/Source/Developer/XmlParser/Private/XmlFile.cpp b/Engine/Source/Developer/XmlParser/Private/XmlFile.cpp
--- a/Engine/Source/Developer/XmlParser/Private/XmlFile.cpp
+++ b/Engine/Source/Developer/XmlParser/Private/XmlFile.cpp
@@ -10,14 +10,21 @@ namespace
 {
 /**
  * Split a buffer of characters into lines
+ * @param OutArray Array of strings to write lines to
+ * @param InBuffer Memory containing text to split
+ * @param EndOfBuffer Pointer to the byte after last byte of the buffer
+ * @param Delim Character used to determine a line ending
  */
 template <typename CharType>
-void SplitLines(TArray<FString>& OutArray, const CharType* InBuffer, uint32 BufferSize, ANSICHAR Delim = '\n');
+void SplitLines(TArray<FString>& OutArray, const CharType* InBuffer, const void* EndOfBuffer, ANSICHAR Delim = '\n');
 
 /**
  * Take an XML buffer, detect the size of character it uses and split into lines
+ * @param OutArray Array of strings to write lines to
+ * @param InBuffer Memory containing text to split
+ * @return False if character type not detected (only guaranteed for XML files)
  */
-bool FindCharSizeAndSplitLines(TArray<FString>& Array, const void* InBuffer, uint32 BufferSize);
+bool FindCharSizeAndSplitLines(TArray<FString>& OutArray, const void* InBuffer, uint32 BufferSize);
 
 }
 
@@ -54,7 +61,7 @@ bool FXmlFile::LoadFile(const FString& InFile, EConstructMethod::Type ConstructM
 		uint32 BufferSize = FileReader->TotalSize();
 		void* Buffer = FMemory::Malloc(BufferSize);
 		FileReader->Serialize(Buffer, BufferSize);
-	
+
 		// Parse file buffer into an array of lines
 		if (!FindCharSizeAndSplitLines(Input, Buffer, BufferSize))
 		{
@@ -71,7 +78,7 @@ bool FXmlFile::LoadFile(const FString& InFile, EConstructMethod::Type ConstructM
 	else
 	{
 		// Parse input buffer into an array of lines
-		SplitLines(Input, *InFile, InFile.Len());
+		SplitLines(Input, *InFile, *InFile + InFile.Len());
 	}
 
 	// Pre-process the input
@@ -145,14 +152,8 @@ FXmlNode* FXmlFile::GetRootNode()
 
 bool FXmlFile::Save(const FString& Path)
 {
-#if  PLATFORM_TCHAR_IS_4_BYTES
-	ErrorMessage += NSLOCTEXT("XmlParser", "NotImplemented", "FXmlFile::Save has not been implemented on this platform").ToString();
-	return false;
-#else
 	FString Xml;
-	// First write the UTF-16 byte order mark
-	Xml.AppendChar(UNICODE_BOM);
-	Xml += TEXT("<?xml version=\"1.0\" encoding=\"UTF-16\"?>") LINE_TERMINATOR;
+	Xml += TEXT("<?xml version=\"1.0\" encoding=\"UTF-8\"?>") LINE_TERMINATOR;
 
 	TArray<const FXmlNode*> NodeStack;
 	static const TCHAR Tabs[] = TEXT("\t\t\t\t\t\t\t\t\t\t\t\t\t\t\t\t");
@@ -192,24 +193,26 @@ bool FXmlFile::Save(const FString& Path)
 		ErrorMessage += FString::Printf(TEXT("\"%s\""), *Path);
 		return false;
 	}
-	Archive->Serialize(&Xml.GetCharArray()[0], Xml.Len() * sizeof(TCHAR));
+
+	FTCHARToUTF8 Converter(*Xml);
+	Archive->Serialize(const_cast<char*>(Converter.Get()), Converter.Length());
 	return true;
-#endif // PLATFORM_TCHAR_IS_1_BYTE || PLATFORM_TCHAR_IS_4_BYTES
 }
 
 namespace
 {
 template <typename CharType>
-void SplitLines(TArray<FString>& Array, const CharType* Buffer, uint32 BufferSize, ANSICHAR Delim)
+void SplitLines(TArray<FString>& OutArray, const CharType* Buffer, const void* EndOfBuffer, ANSICHAR Delim)
 {
+	uint32 CharacterCount = static_cast<const CharType*>(EndOfBuffer) - Buffer;
 	FString WorkingLine;
-	for(uint32 i = 0; i < BufferSize; ++i)
+	for(uint32 i = 0; i != CharacterCount; ++i)
 	{
 		if(Buffer[i] == Delim)
 		{
 			if(WorkingLine.Len())
 			{
-				Array.Add(WorkingLine);
+				OutArray.Add(WorkingLine);
 				WorkingLine = TEXT("");
 			}
 		}
@@ -221,7 +224,7 @@ void SplitLines(TArray<FString>& Array, const CharType* Buffer, uint32 BufferSiz
 
 	if(WorkingLine.Len())
 	{
-		Array.Add(WorkingLine);
+		OutArray.Add(WorkingLine);
 	}
 }
 
@@ -240,19 +243,21 @@ bool FindCharSizeAndSplitLines(TArray<FString>& OutArray, const void* InBuffer,
 		return false;
 	}
 
+	const void* EndOfBuffer = static_cast<const char*>(InBuffer) + BufferSize;
+
 	// Parse input buffer into an array of lines
 	switch (CharCheck.CharacterWidth)
 	{
 	case 1:
-		SplitLines(OutArray, static_cast<const uint8*>(CharCheck.TextStart), BufferSize);
+		SplitLines(OutArray, static_cast<const uint8*>(CharCheck.TextStart), EndOfBuffer);
 		break;
 
 	case 2:
-		SplitLines(OutArray, static_cast<const uint16*>(CharCheck.TextStart), BufferSize);
+		SplitLines(OutArray, static_cast<const uint16*>(CharCheck.TextStart), EndOfBuffer);
 		break;
 
 	case 4:
-		SplitLines(OutArray, static_cast<const uint32*>(CharCheck.TextStart), BufferSize);
+		SplitLines(OutArray, static_cast<const uint32*>(CharCheck.TextStart), EndOfBuffer);
 		break;
 
 	default:
